#UE4 - Add advertising modules to build

diff --git a/Engine/Source/Runtime/Engine/Engine.Build.cs b/Engine/Source/Runtime/Engine/Engine.Build.cs
--- a/Engine/Source/Runtime/Engine/Engine.Build.cs
+++ b/Engine/Source/Runtime/Engine/Engine.Build.cs
@@ -28,6 +28,7 @@ public class Engine : ModuleRules
 				"DerivedDataCache",
 				"ImageWrapper",
 				"HeadMountedDisplay",
+				"Advertising"
 			}
 		);
 
@@ -122,7 +123,8 @@ public class Engine : ModuleRules
 				"Analytics",
 				"AnalyticsET",
 				"OnlineSubsystem", 
-				"OnlineSubsystemUtils"
+				"OnlineSubsystemUtils",
+				"Advertising"
 			}
 		);
 
diff --git a/Engine/Source/Runtime/UE4Game/UE4Game.Build.cs b/Engine/Source/Runtime/UE4Game/UE4Game.Build.cs
--- a/Engine/Source/Runtime/UE4Game/UE4Game.Build.cs
+++ b/Engine/Source/Runtime/UE4Game/UE4Game.Build.cs
@@ -13,6 +13,7 @@ public class UE4Game : ModuleRules
 			PrivateDependencyModuleNames.AddRange(new string[] { "Core", "CoreUObject", "Engine", "OnlineSubsystem", "OnlineSubsystemUtils" });
 			DynamicallyLoadedModuleNames.Add("OnlineSubsystemFacebook");
 			DynamicallyLoadedModuleNames.Add("OnlineSubsystemIOS");
+			DynamicallyLoadedModuleNames.Add("IOSAdvertising");
 		}
 	}
 }
diff --git a/Engine/Source/UE4Game.Target.cs b/Engine/Source/UE4Game.Target.cs
--- a/Engine/Source/UE4Game.Target.cs
+++ b/Engine/Source/UE4Game.Target.cs
@@ -49,6 +49,7 @@ public class UE4GameTarget : TargetRules
 			{
 				OutExtraModuleNames.Add("OnlineSubsystemFacebook");
 				OutExtraModuleNames.Add("OnlineSubsystemIOS");
+				OutExtraModuleNames.Add("IOSAdvertising");
 			}
 			else if (Target.Platform == UnrealTargetPlatform.Android)
 			{
