Fix undefined external symbol errors in UE4Game targets.

#rnx

diff --git a/Engine/Source/Runtime/UE4Game/Private/UE4Game.cpp b/Engine/Source/Runtime/UE4Game/Private/UE4Game.cpp
--- a/Engine/Source/Runtime/UE4Game/Private/UE4Game.cpp
+++ b/Engine/Source/Runtime/UE4Game/Private/UE4Game.cpp
@@ -13,4 +13,6 @@ PER_MODULE_BOILERPLATE
 bool GIsGameAgnosticExe = true;
 TCHAR GInternalProjectName[64] = TEXT("");
 IMPLEMENT_FOREIGN_ENGINE_DIR()
+IMPLEMENT_LIVE_CODING_ENGINE_DIR()
+IMPLEMENT_LIVE_CODING_PROJECT()
 #endif
